Fix Mac CIS

diff --git a/Engine/Plugins/Experimental/CodeEditor/Source/CodeEditor/Private/CPPRichTextSyntaxHighlighterTextLayoutMarshaller.h b/Engine/Plugins/Experimental/CodeEditor/Source/CodeEditor/Private/CPPRichTextSyntaxHighlighterTextLayoutMarshaller.h
--- a/Engine/Plugins/Experimental/CodeEditor/Source/CodeEditor/Private/CPPRichTextSyntaxHighlighterTextLayoutMarshaller.h
+++ b/Engine/Plugins/Experimental/CodeEditor/Source/CodeEditor/Private/CPPRichTextSyntaxHighlighterTextLayoutMarshaller.h
@@ -2,6 +2,7 @@
 #pragma once
 
 #include "SyntaxHighlighterTextLayoutMarshaller.h"
+#include "CodeEditorStyle.h"
 
 /**
  * Get/set the raw text to/from a text layout, and also inject syntax highlighting for our rich-text markup
