Non unity CIS fix

diff --git a/Engine/Plugins/Experimental/CodeEditor/Source/CodeEditor/Private/CodeProjectEditorCommands.cpp b/Engine/Plugins/Experimental/CodeEditor/Source/CodeEditor/Private/CodeProjectEditorCommands.cpp
--- a/Engine/Plugins/Experimental/CodeEditor/Source/CodeEditor/Private/CodeProjectEditorCommands.cpp
+++ b/Engine/Plugins/Experimental/CodeEditor/Source/CodeEditor/Private/CodeProjectEditorCommands.cpp
@@ -2,6 +2,7 @@
 
 #include "CodeEditorPrivatePCH.h"
 #include "CodeProjectEditorCommands.h"
+#include "CodeEditorStyle.h"
 
 #define LOCTEXT_NAMESPACE "CodeProjectEditorCommands"
 
