CodeEditor: Fixed non-unity build; better PCH usage.

diff --git a/Engine/Plugins/Experimental/CodeEditor/Source/CodeEditor/Private/CPPRichTextSyntaxHighlighterTextLayoutMarshaller.h b/Engine/Plugins/Experimental/CodeEditor/Source/CodeEditor/Private/CPPRichTextSyntaxHighlighterTextLayoutMarshaller.h
--- a/Engine/Plugins/Experimental/CodeEditor/Source/CodeEditor/Private/CPPRichTextSyntaxHighlighterTextLayoutMarshaller.h
+++ b/Engine/Plugins/Experimental/CodeEditor/Source/CodeEditor/Private/CPPRichTextSyntaxHighlighterTextLayoutMarshaller.h
@@ -2,7 +2,7 @@
 #pragma once
 
 #include "SyntaxHighlighterTextLayoutMarshaller.h"
-#include "CodeEditorStyle.h"
+
 
 /**
  * Get/set the raw text to/from a text layout, and also inject syntax highlighting for our rich-text markup
diff --git a/Engine/Plugins/Experimental/CodeEditor/Source/CodeEditor/Private/CodeAssetTypeActions.cpp b/Engine/Plugins/Experimental/CodeEditor/Source/CodeEditor/Private/CodeAssetTypeActions.cpp
--- a/Engine/Plugins/Experimental/CodeEditor/Source/CodeEditor/Private/CodeAssetTypeActions.cpp
+++ b/Engine/Plugins/Experimental/CodeEditor/Source/CodeEditor/Private/CodeAssetTypeActions.cpp
@@ -2,11 +2,11 @@
 
 #include "CodeEditorPrivatePCH.h"
 #include "CodeAssetTypeActions.h"
-#include "CodeProjectEditor.h"
-#include "CodeProject.h"
+
 
 #define LOCTEXT_NAMESPACE "CodeAssetTypeActions"
 
+
 FText FCodeAssetTypeActions::GetName() const
 {
 	return LOCTEXT("CodeProjectActionsName", "Code Project");
diff --git a/Engine/Plugins/Experimental/CodeEditor/Source/CodeEditor/Private/CodeEditor.cpp b/Engine/Plugins/Experimental/CodeEditor/Source/CodeEditor/Private/CodeEditor.cpp
--- a/Engine/Plugins/Experimental/CodeEditor/Source/CodeEditor/Private/CodeEditor.cpp
+++ b/Engine/Plugins/Experimental/CodeEditor/Source/CodeEditor/Private/CodeEditor.cpp
@@ -4,8 +4,6 @@
 #include "ModuleManager.h"
 #include "AssetToolsModule.h"
 #include "AssetEditorToolkit.h"
-#include "CodeEditorStyle.h"
-#include "CodeProjectEditor.h"
 #include "LevelEditor.h"
 
 static const FName CodeEditorTabName( TEXT( "CodeEditor" ) );
diff --git a/Engine/Plugins/Experimental/CodeEditor/Source/CodeEditor/Private/CodeEditorPrivatePCH.h b/Engine/Plugins/Experimental/CodeEditor/Source/CodeEditor/Private/CodeEditorPrivatePCH.h
--- a/Engine/Plugins/Experimental/CodeEditor/Source/CodeEditor/Private/CodeEditorPrivatePCH.h
+++ b/Engine/Plugins/Experimental/CodeEditor/Source/CodeEditor/Private/CodeEditorPrivatePCH.h
@@ -2,10 +2,24 @@
 
 #pragma once
 
+
+/* Private dependencies
+ *****************************************************************************/
+
 #include "Core.h"
 #include "CoreUObject.h"
 #include "SlateBasics.h"
 #include "EditorStyle.h"
 #include "PropertyEditing.h"
 #include "UnrealEd.h"
-#include "DirectoryWatcherModule.h"
\ No newline at end of file
+#include "DirectoryWatcherModule.h"
+
+
+/* Private includes
+ *****************************************************************************/
+
+#include "CodeEditorStyle.h"
+#include "CodeProject.h"
+#include "CodeProjectEditorCommands.h"
+#include "CodeProjectEditor.h"
+#include "CodeProjectItem.h"
diff --git a/Engine/Plugins/Experimental/CodeEditor/Source/CodeEditor/Private/CodeEditorStyle.cpp b/Engine/Plugins/Experimental/CodeEditor/Source/CodeEditor/Private/CodeEditorStyle.cpp
--- a/Engine/Plugins/Experimental/CodeEditor/Source/CodeEditor/Private/CodeEditorStyle.cpp
+++ b/Engine/Plugins/Experimental/CodeEditor/Source/CodeEditor/Private/CodeEditorStyle.cpp
@@ -1,7 +1,7 @@
 // Copyright 1998-2014 Epic Games, Inc. All Rights Reserved.
 
 #include "CodeEditorPrivatePCH.h"
-#include "CodeEditorStyle.h"
+
 
 TSharedPtr< FSlateStyleSet > FCodeEditorStyle::StyleSet = nullptr;
 
diff --git a/Engine/Plugins/Experimental/CodeEditor/Source/CodeEditor/Private/CodeProject.cpp b/Engine/Plugins/Experimental/CodeEditor/Source/CodeEditor/Private/CodeProject.cpp
--- a/Engine/Plugins/Experimental/CodeEditor/Source/CodeEditor/Private/CodeProject.cpp
+++ b/Engine/Plugins/Experimental/CodeEditor/Source/CodeEditor/Private/CodeProject.cpp
@@ -1,7 +1,7 @@
 // Copyright 1998-2014 Epic Games, Inc. All Rights Reserved.
 
 #include "CodeEditorPrivatePCH.h"
-#include "CodeProject.h"
+
 
 UCodeProject::UCodeProject(const FObjectInitializer& ObjectInitializer)
 	: Super(ObjectInitializer)
diff --git a/Engine/Plugins/Experimental/CodeEditor/Source/CodeEditor/Private/CodeProjectEditor.cpp b/Engine/Plugins/Experimental/CodeEditor/Source/CodeEditor/Private/CodeProjectEditor.cpp
--- a/Engine/Plugins/Experimental/CodeEditor/Source/CodeEditor/Private/CodeProjectEditor.cpp
+++ b/Engine/Plugins/Experimental/CodeEditor/Source/CodeEditor/Private/CodeProjectEditor.cpp
@@ -1,15 +1,11 @@
 // Copyright 1998-2014 Epic Games, Inc. All Rights Reserved.
 
 #include "CodeEditorPrivatePCH.h"
-#include "CodeProjectEditor.h"
 #include "SCodeEditor.h"
 #include "SCodeProjectEditor.h"
-#include "CodeProject.h"
-#include "CodeProjectItem.h"
-#include "CodeEditorStyle.h"
 #include "SDockTab.h"
 #include "CodeProjectEditorToolbar.h"
-#include "CodeProjectEditorCommands.h"
+
 
 #define LOCTEXT_NAMESPACE "CodeEditor"
 
diff --git a/Engine/Plugins/Experimental/CodeEditor/Source/CodeEditor/Private/CodeProjectEditorCommands.cpp b/Engine/Plugins/Experimental/CodeEditor/Source/CodeEditor/Private/CodeProjectEditorCommands.cpp
--- a/Engine/Plugins/Experimental/CodeEditor/Source/CodeEditor/Private/CodeProjectEditorCommands.cpp
+++ b/Engine/Plugins/Experimental/CodeEditor/Source/CodeEditor/Private/CodeProjectEditorCommands.cpp
@@ -1,20 +1,22 @@
 // Copyright 1998-2015 Epic Games, Inc. All Rights Reserved.
 
 #include "CodeEditorPrivatePCH.h"
-#include "CodeProjectEditorCommands.h"
-#include "CodeEditorStyle.h"
+
 
 #define LOCTEXT_NAMESPACE "CodeProjectEditorCommands"
 
+
 FCodeProjectEditorCommands::FCodeProjectEditorCommands() 
 	: TCommands<FCodeProjectEditorCommands>("CodeEditor", LOCTEXT("General", "General"), NAME_None, FCodeEditorStyle::GetStyleSetName())
 {
 }
 
+
 void FCodeProjectEditorCommands::RegisterCommands()
 {
 	UI_COMMAND(Save, "Save", "Save the currently active document.", EUserInterfaceActionType::Button, FInputGesture(EModifierKey::Control, EKeys::S));
 	UI_COMMAND(SaveAll, "Save All", "Save all open documents.", EUserInterfaceActionType::Button, FInputGesture(EModifierKey::Control | EModifierKey::Shift, EKeys::S));
 }
 
-#undef LOCTEXT_NAMESPACE
\ No newline at end of file
+
+#undef LOCTEXT_NAMESPACE
diff --git a/Engine/Plugins/Experimental/CodeEditor/Source/CodeEditor/Private/CodeProjectEditorToolbar.cpp b/Engine/Plugins/Experimental/CodeEditor/Source/CodeEditor/Private/CodeProjectEditorToolbar.cpp
--- a/Engine/Plugins/Experimental/CodeEditor/Source/CodeEditor/Private/CodeProjectEditorToolbar.cpp
+++ b/Engine/Plugins/Experimental/CodeEditor/Source/CodeEditor/Private/CodeProjectEditorToolbar.cpp
@@ -2,11 +2,10 @@
 
 #include "CodeEditorPrivatePCH.h"
 #include "CodeProjectEditorToolbar.h"
-#include "CodeProjectEditorCommands.h"
 #include "LevelEditorActions.h"
 #include "SourceCodeNavigation.h"
 #include "EditorStyleSet.h"
-#include "CodeProjectEditor.h"
+
 
 void FCodeProjectEditorToolbar::AddEditorToolbar(TSharedPtr<FExtender> Extender)
 {
diff --git a/Engine/Plugins/Experimental/CodeEditor/Source/CodeEditor/Private/CodeProjectFactory.cpp b/Engine/Plugins/Experimental/CodeEditor/Source/CodeEditor/Private/CodeProjectFactory.cpp
--- a/Engine/Plugins/Experimental/CodeEditor/Source/CodeEditor/Private/CodeProjectFactory.cpp
+++ b/Engine/Plugins/Experimental/CodeEditor/Source/CodeEditor/Private/CodeProjectFactory.cpp
@@ -2,10 +2,11 @@
 
 #include "CodeEditorPrivatePCH.h"
 #include "CodeProjectFactory.h"
-#include "CodeProject.h"
+
 
 #define LOCTEXT_NAMESPACE "CodeEditor"
 
+
 UCodeProjectFactory::UCodeProjectFactory(const class FObjectInitializer& ObjectInitializer)
 	: Super(ObjectInitializer)
 {
@@ -14,10 +15,12 @@ UCodeProjectFactory::UCodeProjectFactory(const class FObjectInitializer& ObjectI
 	SupportedClass = UCodeProject::StaticClass();
 }
 
+
 UObject* UCodeProjectFactory::FactoryCreateNew(UClass* Class, UObject* InParent, FName Name, EObjectFlags Flags, UObject* Context, FFeedbackContext* Warn)
 {
 	UCodeProject* NewCodeProject = NewObject<UCodeProject>(InParent, Class, Name, Flags);
 	return NewCodeProject;
 }
 
-#undef LOCTEXT_NAMESPACE
\ No newline at end of file
+
+#undef LOCTEXT_NAMESPACE
diff --git a/Engine/Plugins/Experimental/CodeEditor/Source/CodeEditor/Private/SCodeEditableText.cpp b/Engine/Plugins/Experimental/CodeEditor/Source/CodeEditor/Private/SCodeEditableText.cpp
--- a/Engine/Plugins/Experimental/CodeEditor/Source/CodeEditor/Private/SCodeEditableText.cpp
+++ b/Engine/Plugins/Experimental/CodeEditor/Source/CodeEditor/Private/SCodeEditableText.cpp
@@ -2,7 +2,7 @@
 
 #include "CodeEditorPrivatePCH.h"
 #include "SCodeEditableText.h"
-#include "CodeEditorStyle.h"
+
 
 void SCodeEditableText::Construct( const FArguments& InArgs )
 {
diff --git a/Engine/Plugins/Experimental/CodeEditor/Source/CodeEditor/Private/SCodeEditor.cpp b/Engine/Plugins/Experimental/CodeEditor/Source/CodeEditor/Private/SCodeEditor.cpp
--- a/Engine/Plugins/Experimental/CodeEditor/Source/CodeEditor/Private/SCodeEditor.cpp
+++ b/Engine/Plugins/Experimental/CodeEditor/Source/CodeEditor/Private/SCodeEditor.cpp
@@ -4,12 +4,12 @@
 #include "SCodeEditor.h"
 #include "SMultiLineEditableText.h"
 #include "CPPRichTextSyntaxHighlighterTextLayoutMarshaller.h"
-#include "CodeEditorStyle.h"
-#include "CodeProjectItem.h"
 #include "SCodeEditableText.h"
 
+
 #define LOCTEXT_NAMESPACE "CodeEditor"
 
+
 void SCodeEditor::Construct(const FArguments& InArgs, UCodeProjectItem* InCodeProjectItem)
 {
 	bDirty = false;
diff --git a/Engine/Plugins/Experimental/CodeEditor/Source/CodeEditor/Private/SCodeProjectEditor.cpp b/Engine/Plugins/Experimental/CodeEditor/Source/CodeEditor/Private/SCodeProjectEditor.cpp
--- a/Engine/Plugins/Experimental/CodeEditor/Source/CodeEditor/Private/SCodeProjectEditor.cpp
+++ b/Engine/Plugins/Experimental/CodeEditor/Source/CodeEditor/Private/SCodeProjectEditor.cpp
@@ -3,16 +3,13 @@
 #include "CodeEditorPrivatePCH.h"
 #include "SCodeProjectEditor.h"
 #include "SProjectViewItem.h"
-#include "CodeEditorStyle.h"
-#include "CodeProjectItem.h"
 #include "DirectoryScanner.h"
-#include "CodeProjectEditor.h"
-#include "CodeProject.h"
-#include "CodeProjectItem.h"
 #include "SThrobber.h"
 
+
 #define LOCTEXT_NAMESPACE "CodeProjectEditor"
 
+
 void SCodeProjectEditor::Construct(const FArguments& InArgs, UCodeProject* InCodeProject)
 {
 	check(InCodeProject);
diff --git a/Engine/Plugins/Experimental/CodeEditor/Source/CodeEditor/Private/SProjectViewItem.cpp b/Engine/Plugins/Experimental/CodeEditor/Source/CodeEditor/Private/SProjectViewItem.cpp
--- a/Engine/Plugins/Experimental/CodeEditor/Source/CodeEditor/Private/SProjectViewItem.cpp
+++ b/Engine/Plugins/Experimental/CodeEditor/Source/CodeEditor/Private/SProjectViewItem.cpp
@@ -2,10 +2,11 @@
 
 #include "CodeEditorPrivatePCH.h"
 #include "SProjectViewItem.h"
-#include "CodeEditorStyle.h"
+
 
 #define LOCTEXT_NAMESPACE "ProjectViewItem"
 
+
 void SProjectViewItem::Construct(const FArguments& InArgs)
 {
 	ChildSlot
@@ -32,4 +33,5 @@ void SProjectViewItem::Construct(const FArguments& InArgs)
 	];
 }
 
-#undef LOCTEXT_NAMESPACE
\ No newline at end of file
+
+#undef LOCTEXT_NAMESPACE
