Change Code Editor plugin's module type from Developer to Editor

diff --git a/Engine/Plugins/Experimental/CodeEditor/CodeEditor.uplugin b/Engine/Plugins/Experimental/CodeEditor/CodeEditor.uplugin
--- a/Engine/Plugins/Experimental/CodeEditor/CodeEditor.uplugin
+++ b/Engine/Plugins/Experimental/CodeEditor/CodeEditor.uplugin
@@ -18,7 +18,7 @@
 	[
 		{
 			"Name" : "CodeEditor",
-			"Type" : "Developer",
+			"Type" : "Editor",
 			"LoadingPhase" : "Default"
 		}
 	]
