Converted menu: Code Editor
#rnx

diff --git a/Engine/Plugins/Experimental/CodeEditor/Source/CodeEditor/CodeEditor.Build.cs b/Engine/Plugins/Experimental/CodeEditor/Source/CodeEditor/CodeEditor.Build.cs
--- a/Engine/Plugins/Experimental/CodeEditor/Source/CodeEditor/CodeEditor.Build.cs
+++ b/Engine/Plugins/Experimental/CodeEditor/Source/CodeEditor/CodeEditor.Build.cs
@@ -25,7 +25,8 @@ namespace UnrealBuildTool.Rules
 					"Kismet",  // for FWorkflowCentricApplication
 					"InputCore",
 					"DirectoryWatcher",
-					"LevelEditor"
+					"LevelEditor",
+					"EditorMenus",
 				}
 				);
 		}
diff --git a/Engine/Plugins/Experimental/CodeEditor/Source/CodeEditor/Private/CodeEditor.cpp b/Engine/Plugins/Experimental/CodeEditor/Source/CodeEditor/Private/CodeEditor.cpp
--- a/Engine/Plugins/Experimental/CodeEditor/Source/CodeEditor/Private/CodeEditor.cpp
+++ b/Engine/Plugins/Experimental/CodeEditor/Source/CodeEditor/Private/CodeEditor.cpp
@@ -5,8 +5,7 @@
 #include "Modules/ModuleManager.h"
 #include "Textures/SlateIcon.h"
 #include "Framework/Commands/UIAction.h"
-#include "Framework/MultiBox/MultiBoxExtender.h"
-#include "Framework/MultiBox/MultiBoxBuilder.h"
+#include "EditorMenuSubsystem.h"
 #include "Framework/Docking/TabManager.h"
 #include "Toolkits/IToolkitHost.h"
 #include "CodeEditorStyle.h"
@@ -20,59 +19,55 @@ static const FName CodeEditorTabName( TEXT( "CodeEditor" ) );
 
 class FCodeEditor : public IModuleInterface
 {
+private:
+	static TSharedRef<SDockTab> SpawnCodeEditorTab(const FSpawnTabArgs& TabArgs)
+	{
+		TSharedRef<FCodeProjectEditor> NewCodeProjectEditor(new FCodeProjectEditor());
+		NewCodeProjectEditor->InitCodeEditor(EToolkitMode::Standalone, TSharedPtr<class IToolkitHost>(), GetMutableDefault<UCodeProject>());
+
+		return FGlobalTabmanager::Get()->GetMajorTabForTabManager(NewCodeProjectEditor->GetTabManager().ToSharedRef()).ToSharedRef();
+	}
+
+	static void OpenCodeEditor()
+	{
+		SpawnCodeEditorTab(FSpawnTabArgs(TSharedPtr<SWindow>(), FTabId()));	
+	}
+
 public:
-	/** IModuleInterface implementation */
-	virtual void StartupModule() override
+
+	virtual void OnPostEngineInit()
 	{
-		FCodeEditorStyle::Initialize();
+		UEditorMenu* Menu = UEditorMenuSubsystem::Get()->ExtendMenu("LevelEditor.MainMenu.File");
+		FEditorMenuSection& Section = Menu->FindOrAddSection("FileProject");
 
-		struct Local
+		FEditorMenuOwnerScoped OwnerScoped(this);
 		{
-			static TSharedRef<SDockTab> SpawnCodeEditorTab(const FSpawnTabArgs& TabArgs)
-			{
-				TSharedRef<FCodeProjectEditor> NewCodeProjectEditor(new FCodeProjectEditor());
-				NewCodeProjectEditor->InitCodeEditor(EToolkitMode::Standalone, TSharedPtr<class IToolkitHost>(), GetMutableDefault<UCodeProject>());
-
-				return FGlobalTabmanager::Get()->GetMajorTabForTabManager(NewCodeProjectEditor->GetTabManager().ToSharedRef()).ToSharedRef();
-			}
-
-			static void OpenCodeEditor()
-			{
-				SpawnCodeEditorTab(FSpawnTabArgs(TSharedPtr<SWindow>(), FTabId()));	
-			}
-
-			static void ExtendMenu(class FMenuBuilder& MenuBuilder)
-			{
-				MenuBuilder.AddMenuEntry
+			FEditorMenuEntry& MenuEntry = Section.AddMenuEntry(
+				"EditSourceCode",
+				LOCTEXT("CodeEditorTabTitle", "Edit Source Code"),
+				LOCTEXT("CodeEditorTooltipText", "Open the Code Editor tab."),
+				FSlateIcon(FCodeEditorStyle::Get().GetStyleSetName(), "CodeEditor.TabIcon"),
+				FUIAction
 				(
-					LOCTEXT( "CodeEditorTabTitle", "Edit Source Code" ),
-					LOCTEXT( "CodeEditorTooltipText", "Open the Code Editor tab." ),
-					FSlateIcon(FCodeEditorStyle::Get().GetStyleSetName(), "CodeEditor.TabIcon"),
-					FUIAction
-					(
-						FExecuteAction::CreateStatic(&Local::OpenCodeEditor)
-					)
-				);
-			}
-		};
-
-		Extender = MakeShareable(new FExtender());
-	
-		// Add code editor extension to main menu
-		Extender->AddMenuExtension(
-			"FileProject",
-			EExtensionHook::First,
-			TSharedPtr<FUICommandList>(),
-			FMenuExtensionDelegate::CreateStatic( &Local::ExtendMenu ) );
-
-		FLevelEditorModule& LevelEditorModule = FModuleManager::LoadModuleChecked<FLevelEditorModule>("LevelEditor");
-		LevelEditorModule.GetMenuExtensibilityManager()->AddExtender( Extender );
+					FExecuteAction::CreateStatic(&FCodeEditor::OpenCodeEditor)
+				)
+			);
+			MenuEntry.InsertPosition = FEditorMenuInsert(NAME_None, EEditorMenuInsertType::First);
+		}
+	}
+
+	/** IModuleInterface implementation */
+	virtual void StartupModule() override
+	{
+		FCodeEditorStyle::Initialize();
 
 		// Register a tab spawner so that our tab can be automatically restored from layout files
-		FGlobalTabmanager::Get()->RegisterTabSpawner( CodeEditorTabName, FOnSpawnTab::CreateStatic( &Local::SpawnCodeEditorTab ) )
+		FGlobalTabmanager::Get()->RegisterTabSpawner( CodeEditorTabName, FOnSpawnTab::CreateStatic( &FCodeEditor::SpawnCodeEditorTab ) )
 				.SetDisplayName( LOCTEXT( "CodeEditorTabTitle", "Edit Source Code" ) )
 				.SetTooltipText( LOCTEXT( "CodeEditorTooltipText", "Open the Code Editor tab." ) )
 				.SetIcon(FSlateIcon(FCodeEditorStyle::Get().GetStyleSetName(), "CodeEditor.TabIcon"));
+
+		FCoreDelegates::OnPostEngineInit.AddRaw(this, &FCodeEditor::OnPostEngineInit);
 	}
 
 	virtual void ShutdownModule() override
@@ -80,17 +75,11 @@ public:
 		// Unregister the tab spawner
 		FGlobalTabmanager::Get()->UnregisterTabSpawner( CodeEditorTabName );
 
-		if(FModuleManager::Get().IsModuleLoaded("LevelEditor"))
-		{
-			FLevelEditorModule& LevelEditorModule = FModuleManager::GetModuleChecked<FLevelEditorModule>("LevelEditor");
-			LevelEditorModule.GetMenuExtensibilityManager()->RemoveExtender( Extender );		
-		}
+		FCoreDelegates::OnPostEngineInit.RemoveAll(this);
+		UEditorMenuSubsystem::UnregisterOwner(this);
 
 		FCodeEditorStyle::Shutdown();
 	}
-
-private:
-	TSharedPtr<FExtender> Extender; 
 };
 
 IMPLEMENT_MODULE( FCodeEditor, CodeEditor )
