Fix linux linker error
#rnx

diff --git a/Engine/Plugins/Experimental/CodeEditor/Source/CodeEditor/CodeEditor.Build.cs b/Engine/Plugins/Experimental/CodeEditor/Source/CodeEditor/CodeEditor.Build.cs
--- a/Engine/Plugins/Experimental/CodeEditor/Source/CodeEditor/CodeEditor.Build.cs
+++ b/Engine/Plugins/Experimental/CodeEditor/Source/CodeEditor/CodeEditor.Build.cs
@@ -26,6 +26,7 @@ namespace UnrealBuildTool.Rules
 					"InputCore",
 					"DirectoryWatcher",
 					"LevelEditor",
+					"Engine",
 					"EditorMenus",
 				}
 				);
