Fix CodeEditor menu registration
#rnx

diff --git a/Engine/Plugins/Experimental/CodeEditor/Source/CodeEditor/Private/CodeEditor.cpp b/Engine/Plugins/Experimental/CodeEditor/Source/CodeEditor/Private/CodeEditor.cpp
--- a/Engine/Plugins/Experimental/CodeEditor/Source/CodeEditor/Private/CodeEditor.cpp
+++ b/Engine/Plugins/Experimental/CodeEditor/Source/CodeEditor/Private/CodeEditor.cpp
@@ -37,22 +37,25 @@ public:
 
 	virtual void OnPostEngineInit()
 	{
-		UEditorMenu* Menu = UEditorMenuSubsystem::Get()->ExtendMenu("LevelEditor.MainMenu.File");
-		FEditorMenuSection& Section = Menu->FindOrAddSection("FileProject");
-
-		FEditorMenuOwnerScoped OwnerScoped(this);
+		if (UEditorMenuSubsystem::IsRunningEditorUI())
 		{
-			FEditorMenuEntry& MenuEntry = Section.AddMenuEntry(
-				"EditSourceCode",
-				LOCTEXT("CodeEditorTabTitle", "Edit Source Code"),
-				LOCTEXT("CodeEditorTooltipText", "Open the Code Editor tab."),
-				FSlateIcon(FCodeEditorStyle::Get().GetStyleSetName(), "CodeEditor.TabIcon"),
-				FUIAction
-				(
-					FExecuteAction::CreateStatic(&FCodeEditor::OpenCodeEditor)
-				)
-			);
-			MenuEntry.InsertPosition = FEditorMenuInsert(NAME_None, EEditorMenuInsertType::First);
+			UEditorMenu* Menu = UEditorMenuSubsystem::Get()->ExtendMenu("LevelEditor.MainMenu.File");
+			FEditorMenuSection& Section = Menu->FindOrAddSection("FileProject");
+
+			FEditorMenuOwnerScoped OwnerScoped(this);
+			{
+				FEditorMenuEntry& MenuEntry = Section.AddMenuEntry(
+					"EditSourceCode",
+					LOCTEXT("CodeEditorTabTitle", "Edit Source Code"),
+					LOCTEXT("CodeEditorTooltipText", "Open the Code Editor tab."),
+					FSlateIcon(FCodeEditorStyle::Get().GetStyleSetName(), "CodeEditor.TabIcon"),
+					FUIAction
+					(
+						FExecuteAction::CreateStatic(&FCodeEditor::OpenCodeEditor)
+					)
+				);
+				MenuEntry.InsertPosition = FEditorMenuInsert(NAME_None, EEditorMenuInsertType::First);
+			}
 		}
 	}
 
